CopyBuffer now always runs on the decoder thread and is no longer outsourced to a async task in order to avoid a race condition with the render thread

Removing async copy buffer
#rnx

#Synced-CL 49940788

diff --git a/Engine/Plugins/Media/AvidDNxHDDecoderElectra/Source/Private/Windows/ElectraMediaAvidDNxHDDecoder.cpp b/Engine/Plugins/Media/AvidDNxHDDecoderElectra/Source/Private/Windows/ElectraMediaAvidDNxHDDecoder.cpp
--- a/Engine/Plugins/Media/AvidDNxHDDecoderElectra/Source/Private/Windows/ElectraMediaAvidDNxHDDecoder.cpp
+++ b/Engine/Plugins/Media/AvidDNxHDDecoderElectra/Source/Private/Windows/ElectraMediaAvidDNxHDDecoder.cpp
@@ -146,7 +146,7 @@ public:
 #if ELECTRA_MEDIAGPUBUFFER_DX12
 		if (InBufferIndex == 0)
 		{
-			SyncObject = { GPUBuffer.Fence.GetReference(), GPUBuffer.FenceValue, nullptr, GPUBuffer_TaskSync };
+			SyncObject = { GPUBuffer.Fence.GetReference(), GPUBuffer.FenceValue, nullptr };
 			return true;
 		}
 #endif
@@ -201,7 +201,6 @@ public:
 	EElectraTextureSamplePixelEncoding BufferEncoding = EElectraTextureSamplePixelEncoding::Native;
 #if ELECTRA_MEDIAGPUBUFFER_DX12
 	FElectraMediaDecoderOutputBufferPool_DX12::FOutputData GPUBuffer;
-	TSharedPtr<FElectraDecodersPlatformResources::FAsyncConsecutiveTaskSync, ESPMode::ThreadSafe> GPUBuffer_TaskSync;
 #endif
 };
 
@@ -286,7 +285,6 @@ private:
 	uint32 MaxOutputBuffers;
 #if ELECTRA_MEDIAGPUBUFFER_DX12
 	mutable TSharedPtr<FElectraMediaDecoderOutputBufferPool_DX12, ESPMode::ThreadSafe> D3D12ResourcePool;
-	mutable TSharedPtr<FElectraDecodersPlatformResources::FAsyncConsecutiveTaskSync, ESPMode::ThreadSafe> TaskSync;
 #endif
 };
 
@@ -733,12 +731,6 @@ IElectraDecoder::EDecoderError FVideoDecoderAvidDNxHDElectra::DecodeAccessUnit(c
 				D3D12ResourcePool = MakeShared<FElectraMediaDecoderOutputBufferPool_DX12, ESPMode::ThreadSafe>(D3D12Device, MaxOutputBuffers, NewOutput->Width, AllocHeight, BPP);
 			}
 
-			// Create a tasksync instance so we can have our own async jobs run in consecutive order
-			if (!TaskSync.IsValid())
-			{
-				TaskSync = FElectraDecodersPlatformResources::CreateAsyncConsecutiveTaskSync();
-			}
-
 			// Request resource and fence...
 			if (D3D12ResourcePool->AllocateOutputDataAsBuffer(HResult, HResultMsg, NewOutput->GPUBuffer, ResourcePitch))
 			{
@@ -795,20 +787,8 @@ IElectraDecoder::EDecoderError FVideoDecoderAvidDNxHDElectra::DecodeAccessUnit(c
 				// If we have a temp buffer, we need to copy the data first...
 				if (!TempBuffer.IsEmpty())
 				{
-					if (FElectraDecodersPlatformResources::RunCodeAsync([BufferResource = NewOutput->GPUBuffer.Resource, Width = NewOutput->Width, Height = NewOutput->Height, ResourcePitch, BPP, TempBuffer, BufferFence = NewOutput->GPUBuffer.Fence, BufferFenceValue = NewOutput->GPUBuffer.FenceValue]()
-						{
-							CopyData(BufferResource, ResourcePitch, TempBuffer.GetData(), Width * BPP, Height);
-							BufferFence->Signal(BufferFenceValue);
-						}, TaskSync))
-					{
-						NewOutput->GPUBuffer_TaskSync = TaskSync;
-					}
-					else
-					{
-						// Async copy failed, do it synchronously...
-						CopyData(NewOutput->GPUBuffer.Resource, ResourcePitch, TempBuffer.GetData(), NewOutput->Width * BPP, NewOutput->Height);
-						NewOutput->GPUBuffer.Fence->Signal(NewOutput->GPUBuffer.FenceValue);
-					}
+					CopyData(NewOutput->GPUBuffer.Resource, ResourcePitch, TempBuffer.GetData(), NewOutput->Width * BPP, NewOutput->Height);
+					NewOutput->GPUBuffer.Fence->Signal(NewOutput->GPUBuffer.FenceValue);
 				}
 				else
 				{
diff --git a/Engine/Plugins/Media/ElectraCodecs/Source/ElectraDecoders/Private/DX/VideoDecoderH264_DX.cpp b/Engine/Plugins/Media/ElectraCodecs/Source/ElectraDecoders/Private/DX/VideoDecoderH264_DX.cpp
--- a/Engine/Plugins/Media/ElectraCodecs/Source/ElectraDecoders/Private/DX/VideoDecoderH264_DX.cpp
+++ b/Engine/Plugins/Media/ElectraCodecs/Source/ElectraDecoders/Private/DX/VideoDecoderH264_DX.cpp
@@ -167,7 +167,7 @@ public:
 			SyncObject = { Buffer.DecoderSync.GetReference(), 0, Buffer.MFSample };
 			// note: a reference to the MFSample remains with this instance, but is assumed to soon be released as the instance will be destroyed shortly after this was called
 #else
-			SyncObject = { Buffer.Fence.GetReference(), Buffer.FenceValue, nullptr, Buffer_TaskSync };
+			SyncObject = { Buffer.Fence.GetReference(), Buffer.FenceValue, nullptr };
 #endif
 			return true;
 		}
@@ -233,7 +233,6 @@ public:
 	} Buffer;
 #else
 	FElectraMediaDecoderOutputBufferPool_DX12::FOutputData Buffer;
-	TSharedPtr<FElectraDecodersPlatformResources::FAsyncConsecutiveTaskSync, ESPMode::ThreadSafe> Buffer_TaskSync;
 #endif
 };
 
@@ -373,7 +372,6 @@ private:
 	uint32 MaxHeight;
 #if !ALLOW_MFSAMPLE_WITH_DX12
 	mutable TSharedPtr<FElectraMediaDecoderOutputBufferPool_DX12, ESPMode::ThreadSafe> D3D12ResourcePool;
-	mutable TSharedPtr<FElectraDecodersPlatformResources::FAsyncConsecutiveTaskSync, ESPMode::ThreadSafe> TaskSync;
 #endif
 
 	static const GUID MFTms_AVLowLatencyMode;
@@ -1181,12 +1179,6 @@ bool FElectraVideoDecoderH264_DX::ConvertDecoderOutput()
 			D3D12ResourcePool = MakeShared<FElectraMediaDecoderOutputBufferPool_DX12, ESPMode::ThreadSafe>(D3D12Device, MaxOutputBuffers, MaxWidth, MaxHeight * 3 / 2, 1);
 		}
 
-		// Create a tasksync instance so we can have our own async jobs run in consecutive order
-		if (!TaskSync.IsValid())
-		{
-			TaskSync = FElectraDecodersPlatformResources::CreateAsyncConsecutiveTaskSync();
-		}
-
 		// Get data from MFSample...
 		TRefCountPtr<IMF2DBuffer> Buffer2D;
 		if ((Result = Buffer->QueryInterface(__uuidof(IMF2DBuffer), (void**)Buffer2D.GetInitReference())) != S_OK)
@@ -1204,18 +1196,7 @@ bool FElectraVideoDecoderH264_DX::ConvertDecoderOutput()
 		}
 
 		NewOutput->Pitch = BufferPitch;
-
-		if (FElectraDecodersPlatformResources::RunCodeAsync([DecodedOutputSample, Buffer2D, BufferResource = NewOutput->Buffer.Resource, BufferFence = NewOutput->Buffer.Fence, BufferFenceValue = NewOutput->Buffer.FenceValue, BufferPitch, DecodedHeight = NewOutput->DecodedHeight]()
-			{
-				CopyBufferData(Buffer2D, BufferResource, BufferFence, BufferFenceValue, BufferPitch, DecodedHeight);
-			}, TaskSync))
-		{
-			NewOutput->Buffer_TaskSync = TaskSync;
-		}
-		else
-		{
-			CopyBufferData(Buffer2D, NewOutput->Buffer.Resource, NewOutput->Buffer.Fence, NewOutput->Buffer.FenceValue, BufferPitch, NewOutput->DecodedHeight);
-		}
+		CopyBufferData(Buffer2D, NewOutput->Buffer.Resource, NewOutput->Buffer.Fence, NewOutput->Buffer.FenceValue, BufferPitch, NewOutput->DecodedHeight);
 #else
 		//
 		// Post SDK 22621: direct link from MFSample to DX12 resource
diff --git a/Engine/Plugins/Media/ElectraCodecs/Source/ElectraDecoders/Private/DX/VideoDecoderH265_DX.cpp b/Engine/Plugins/Media/ElectraCodecs/Source/ElectraDecoders/Private/DX/VideoDecoderH265_DX.cpp
--- a/Engine/Plugins/Media/ElectraCodecs/Source/ElectraDecoders/Private/DX/VideoDecoderH265_DX.cpp
+++ b/Engine/Plugins/Media/ElectraCodecs/Source/ElectraDecoders/Private/DX/VideoDecoderH265_DX.cpp
@@ -166,7 +166,7 @@ public:
 			SyncObject = { Buffer.DecoderSync.GetReference(), 0, Buffer.MFSample };
 			// note: a reference to the MFSample remains with this instance, but is assumed to soon be released as the instance will be destroyed shortly after this was called
 #else
-			SyncObject = { Buffer.Fence.GetReference(), Buffer.FenceValue, nullptr, Buffer_TaskSync };
+			SyncObject = { Buffer.Fence.GetReference(), Buffer.FenceValue, nullptr };
 #endif
 			return true;
 		}
@@ -232,7 +232,6 @@ public:
 	} Buffer;
 #else
 	FElectraMediaDecoderOutputBufferPool_DX12::FOutputData Buffer;
-	TSharedPtr<FElectraDecodersPlatformResources::FAsyncConsecutiveTaskSync, ESPMode::ThreadSafe> Buffer_TaskSync;
 #endif
 };
 
@@ -370,7 +369,6 @@ private:
 	uint32 MaxHeight;
 #if !ALLOW_MFSAMPLE_WITH_DX12
 	mutable TSharedPtr<FElectraMediaDecoderOutputBufferPool_DX12, ESPMode::ThreadSafe> D3D12ResourcePool;
-	mutable TSharedPtr<FElectraDecodersPlatformResources::FAsyncConsecutiveTaskSync, ESPMode::ThreadSafe> TaskSync;
 #endif
 
 	static const GUID MFTms_AVLowLatencyMode;
@@ -1187,12 +1185,6 @@ bool FElectraVideoDecoderH265_DX::ConvertDecoderOutput()
 			D3D12ResourcePool = MakeShared<FElectraMediaDecoderOutputBufferPool_DX12, ESPMode::ThreadSafe>(D3D12Device, MaxOutputBuffers, MaxWidth, MaxHeight * 3 / 2, (PixFmt == EPixelFormat::PF_NV12) ? 1 : 2);
 		}
 
-		// Create a tasksync instance so we can have our own async jobs run in consecutive order
-		if (!TaskSync.IsValid())
-		{
-			TaskSync = FElectraDecodersPlatformResources::CreateAsyncConsecutiveTaskSync();
-		}
-
 		// Get data from MFSample...
 		TRefCountPtr<IMF2DBuffer> Buffer2D;
 		if ((Result = Buffer->QueryInterface(__uuidof(IMF2DBuffer), (void**)Buffer2D.GetInitReference())) != S_OK)
@@ -1210,18 +1202,7 @@ bool FElectraVideoDecoderH265_DX::ConvertDecoderOutput()
 		}
 
 		NewOutput->Pitch = BufferPitch;
-
-		if (FElectraDecodersPlatformResources::RunCodeAsync([DecodedOutputSample, Buffer2D, BufferResource = NewOutput->Buffer.Resource, BufferFence = NewOutput->Buffer.Fence, BufferFenceValue = NewOutput->Buffer.FenceValue, BufferPitch, DecodedHeight = NewOutput->DecodedHeight]()
-			{
-				CopyBufferData(Buffer2D, BufferResource, BufferFence, BufferFenceValue, BufferPitch, DecodedHeight);
-			}, TaskSync))
-		{
-			NewOutput->Buffer_TaskSync = TaskSync;
-		}
-		else
-		{
-			CopyBufferData(Buffer2D, NewOutput->Buffer.Resource, NewOutput->Buffer.Fence, NewOutput->Buffer.FenceValue, BufferPitch, NewOutput->DecodedHeight);
-		}
+		CopyBufferData(Buffer2D, NewOutput->Buffer.Resource, NewOutput->Buffer.Fence, NewOutput->Buffer.FenceValue, BufferPitch, NewOutput->DecodedHeight);
 #else
 		//
 		// Post SDK 22621: direct link from MFSample to DX12 resource
diff --git a/Engine/Plugins/Media/ElectraCodecs/Source/ElectraDecoders/Private/Windows/WindowsElectraDecodersPlatformResources.cpp b/Engine/Plugins/Media/ElectraCodecs/Source/ElectraDecoders/Private/Windows/WindowsElectraDecodersPlatformResources.cpp
--- a/Engine/Plugins/Media/ElectraCodecs/Source/ElectraDecoders/Private/Windows/WindowsElectraDecodersPlatformResources.cpp
+++ b/Engine/Plugins/Media/ElectraCodecs/Source/ElectraDecoders/Private/Windows/WindowsElectraDecodersPlatformResources.cpp
@@ -81,41 +81,6 @@ bool FElectraDecodersPlatformResourcesWindows::GetD3DDeviceAndVersion(void** Out
 	}
 }
 
-
-
-class FElectraDecodersPlatformResourcesWindows::FAsyncConsecutiveTaskSync
-{
-public:
-	FAsyncConsecutiveTaskSync() = default;
-	~FAsyncConsecutiveTaskSync() = default;
-	FGraphEventRef GraphEvent;
-};
-
-TSharedPtr<FElectraDecodersPlatformResourcesWindows::FAsyncConsecutiveTaskSync, ESPMode::ThreadSafe> FElectraDecodersPlatformResourcesWindows::CreateAsyncConsecutiveTaskSync()
-{
-	return MakeShared<FElectraDecodersPlatformResourcesWindows::FAsyncConsecutiveTaskSync, ESPMode::ThreadSafe>();
-}
-
-bool FElectraDecodersPlatformResourcesWindows::RunCodeAsync(TFunction<void()>&& InCodeToRun, TSharedPtr<FElectraDecodersPlatformResourcesWindows::FAsyncConsecutiveTaskSync, ESPMode::ThreadSafe> InTaskSync)
-{
-	FScopeLock Lock(&PlatformElectraDecodersWindows::AsyncJobAccessCS);
-
-	// Execute code async
-	// (We assume this to be code copying buffer data around. Hence we allow only ONE at any time to not clog up
-	//  the buses even more and delay the copy process further)
-	FGraphEventArray Events;
-	auto& Event = InTaskSync ? InTaskSync->GraphEvent : PlatformElectraDecodersWindows::RunCodeAsyncEvent;
-	if (Event.IsValid())
-	{
-		Events.Add(Event);
-	}
-	Event = FFunctionGraphTask::CreateAndDispatchWhenReady(MoveTemp(InCodeToRun), GET_STATID(STAT_ElectraDecoderAsyncJob), &Events);
-	return true;
-}
-
-
-
-
 static void TriggerDataCopy(TRefCountPtr<ID3D12Resource> InDstResource, TRefCountPtr<ID3D12Resource> InSrcResource, TRefCountPtr<ID3D12GraphicsCommandList> D3DCmdList, TRefCountPtr<ID3D12CommandAllocator> D3DCmdAllocator, TRefCountPtr<ID3D12Fence> D3DFence, uint64 FenceValue, const TSharedPtr<std::atomic<bool>>& FenceSignalled, const FElectraDecoderOutputSync& OutputSync)
 {
 	// Trigger copy (this will eventually execute on the submission thread of RHI if running in UE)
@@ -452,22 +417,8 @@ static bool SetupOutputTextureSampleFromDX12Resource(FString& OutErrorMessage, T
 	InOutTextureSampleToSetup->D3DCmdList->Close();
 
 	// - - - - - - - - - - - - - - - - - - - - - - - - - - - - - - - - - - - - - - - - - - - - - - - - - - - - - - - -
-	// Can we wait for the decoder to have the data ready using an async CPU job?
-	bool bTriggerOk = false;
-	if (OutputSync.TaskSync.IsValid())
-	{
-		// Note: we capture "this" as we ensure that this instance only dies once the copy triggered here is actually done, hence ensuring any reference to "this" is done
-		bTriggerOk = FElectraDecodersPlatformResources::RunCodeAsync([DstResource=InOutTextureSampleToSetup->TextureDX12, SrcResource=InResource, D3DCmdList=InOutTextureSampleToSetup->D3DCmdList, D3DCmdAllocator=InOutTextureSampleToSetup->D3DCmdAllocator, D3DFence=InOutTextureSampleToSetup->D3DFence, FenceValue=InOutTextureSampleToSetup->FenceValue, FenceSignalled=InOutTextureSampleToSetup->FenceSignalled, OutputSync]()
-			{
-				TriggerDataCopy(DstResource, SrcResource, D3DCmdList, D3DCmdAllocator, D3DFence, FenceValue, FenceSignalled, OutputSync);
-			}, OutputSync.TaskSync);
-	}
-
-	if (!bTriggerOk)
-	{
-		// We could not run the trigger async. Schedule the copy right away. Any needed synchronization will be done in the copy-queue by the GPU
-		TriggerDataCopy(InOutTextureSampleToSetup->TextureDX12, InResource, InOutTextureSampleToSetup->D3DCmdList, InOutTextureSampleToSetup->D3DCmdAllocator, InOutTextureSampleToSetup->D3DFence, InOutTextureSampleToSetup->FenceValue, InOutTextureSampleToSetup->FenceSignalled, OutputSync);
-	}
+	// Schedule the copy right away. Any needed synchronization will be done in the copy-queue by the GPU
+	TriggerDataCopy(InOutTextureSampleToSetup->TextureDX12, InResource, InOutTextureSampleToSetup->D3DCmdList, InOutTextureSampleToSetup->D3DCmdAllocator, InOutTextureSampleToSetup->D3DFence, InOutTextureSampleToSetup->FenceValue, InOutTextureSampleToSetup->FenceSignalled, OutputSync);
 
 	// Finish initialization.
 	if (!InOutTextureSampleToSetup->FinishInitialization())
diff --git a/Engine/Plugins/Media/ElectraCodecs/Source/ElectraDecoders/Public/Windows/WindowsElectraDecoderPlatformOutputHandleTypes.h b/Engine/Plugins/Media/ElectraCodecs/Source/ElectraDecoders/Public/Windows/WindowsElectraDecoderPlatformOutputHandleTypes.h
--- a/Engine/Plugins/Media/ElectraCodecs/Source/ElectraDecoders/Public/Windows/WindowsElectraDecoderPlatformOutputHandleTypes.h
+++ b/Engine/Plugins/Media/ElectraCodecs/Source/ElectraDecoders/Public/Windows/WindowsElectraDecoderPlatformOutputHandleTypes.h
@@ -26,8 +26,8 @@ class FElectraDecoderOutputSync final
 {
 public:
 	FElectraDecoderOutputSync() : SyncValue(0), CopyDoneSyncValue(0) {}
-	FElectraDecoderOutputSync(TRefCountPtr<IUnknown> InSync, uint64 InSyncValue, TRefCountPtr<IMFSample> InMFSample = nullptr, TSharedPtr<FElectraDecodersPlatformResources::FAsyncConsecutiveTaskSync, ESPMode::ThreadSafe> InTaskSync = nullptr)
-		: Sync(InSync), SyncValue(InSyncValue), TaskSync(InTaskSync), MFSample(InMFSample)
+	FElectraDecoderOutputSync(TRefCountPtr<IUnknown> InSync, uint64 InSyncValue, TRefCountPtr<IMFSample> InMFSample = nullptr)
+		: Sync(InSync), SyncValue(InSyncValue), CopyDoneSyncValue(0), MFSample(InMFSample)
 	{}
 	FElectraDecoderOutputSync(TRefCountPtr<IUnknown> InSync, uint64 InSyncValue, TRefCountPtr<ID3D12Fence> InCopyDoneSync, uint64 InCopyDoneSyncValue)
 		: Sync(InSync), SyncValue(InSyncValue), CopyDoneSync(InCopyDoneSync), CopyDoneSyncValue(InCopyDoneSyncValue)
@@ -35,7 +35,6 @@ public:
 
 	TRefCountPtr<IUnknown>		Sync;														// Optional IUnknown based sync primitive (fence, or DX12 specific WMF object)
 	uint64						SyncValue;													// Optionally needed sync value
-	TSharedPtr<FElectraDecodersPlatformResources::FAsyncConsecutiveTaskSync, ESPMode::ThreadSafe> TaskSync;	// Optional Sync/Perquisite object to use to run code async after the decoder data arrives
 
 	TRefCountPtr<ID3D12Fence>	CopyDoneSync;												// Optional fence to signal end of copy from this buffer back to decoder
 	uint64						CopyDoneSyncValue;											// Value to be used with CopyDoneSync fence
diff --git a/Engine/Plugins/Media/HAPDecoderElectra/Source/HAPDecoderElectra/Private/HAPDecoder/ElectraMediaHAPDecoder.cpp b/Engine/Plugins/Media/HAPDecoderElectra/Source/HAPDecoderElectra/Private/HAPDecoder/ElectraMediaHAPDecoder.cpp
--- a/Engine/Plugins/Media/HAPDecoderElectra/Source/HAPDecoderElectra/Private/HAPDecoder/ElectraMediaHAPDecoder.cpp
+++ b/Engine/Plugins/Media/HAPDecoderElectra/Source/HAPDecoderElectra/Private/HAPDecoder/ElectraMediaHAPDecoder.cpp
@@ -149,12 +149,12 @@ public:
 #if ELECTRA_MEDIAGPUBUFFER_DX12
 		if (InBufferIndex == 0)
 		{
-			SyncObject = { Buffer[Buffer_Color].GPUBuffer.Fence.GetReference(), Buffer[Buffer_Color].GPUBuffer.FenceValue, nullptr, Buffer_TaskSync };
+			SyncObject = { Buffer[Buffer_Color].GPUBuffer.Fence.GetReference(), Buffer[Buffer_Color].GPUBuffer.FenceValue, nullptr };
 			return true;
 		}
 		if (InBufferIndex == 1)
 		{
-			SyncObject = { Buffer[Buffer_Alpha].GPUBuffer.Fence.GetReference(), Buffer[Buffer_Alpha].GPUBuffer.FenceValue, nullptr, Buffer_TaskSync };
+			SyncObject = { Buffer[Buffer_Alpha].GPUBuffer.Fence.GetReference(), Buffer[Buffer_Alpha].GPUBuffer.FenceValue, nullptr };
 			return true;
 		}
 #endif
@@ -234,9 +234,6 @@ public:
 #endif
 	};
 	FBufferInfo Buffer[Buffer_Max];
-#if ELECTRA_MEDIAGPUBUFFER_DX12
-	TSharedPtr<FElectraDecodersPlatformResources::FAsyncConsecutiveTaskSync, ESPMode::ThreadSafe> Buffer_TaskSync;
-#endif
 };
 
 
@@ -314,9 +311,6 @@ private:
 #endif
 	};
 	FBufferInstanceInfo InstanceBufferData[FVideoDecoderOutputHAPElectra::Buffer_Max];
-#if ELECTRA_MEDIAGPUBUFFER_DX12
-	mutable TSharedPtr<FElectraDecodersPlatformResources::FAsyncConsecutiveTaskSync, ESPMode::ThreadSafe> TaskSync;
-#endif
 };
 
 
@@ -688,12 +682,6 @@ IElectraDecoder::EDecoderError FVideoDecoderHAPElectra::DecodeAccessUnit(const F
 					InstanceBufferData[BufIdx].D3D12ResourcePool = MakeShared<FElectraMediaDecoderOutputBufferPool_DX12, ESPMode::ThreadSafe>(D3D12Device, MaxOutputBuffers, DecodedWidth / BlockSizeX, DecodedHeight / BlockSizeY, BlockBytes);
 				}
 
-				// Create a tasksync instance so we can have our own async jobs run in consecutive order
-				if (!TaskSync.IsValid())
-				{
-					TaskSync = FElectraDecodersPlatformResources::CreateAsyncConsecutiveTaskSync();
-				}
-
 				if (InstanceBufferData[BufIdx].D3D12ResourcePool->AllocateOutputDataAsBuffer(HResult, HResultMsg, BufferInfo.GPUBuffer, ResourcePitch))
 				{
 					// Make the decoder write to a cached, temp buffer (otherwise performance is quite bad)
@@ -742,34 +730,15 @@ IElectraDecoder::EDecoderError FVideoDecoderHAPElectra::DecodeAccessUnit(const F
 			{
 				if (Result == HapResult_No_Error)
 				{
-					if (FElectraDecodersPlatformResources::RunCodeAsync([BufferResource = BufferInfo.GPUBuffer.Resource, BufferFence = BufferInfo.GPUBuffer.Fence, BufferFenceValue = BufferInfo.GPUBuffer.FenceValue, ResourcePitch, BufferAddr, BufferPitch, DecodedPitch = DecodedHeight / BlockSizeY, TempBuffer]()
-						{
-							uint8* DX12BufferAddr;
-							HRESULT Res = BufferResource->Map(0, nullptr, (void**)&DX12BufferAddr);
-							check(SUCCEEDED(Res));
+					uint8* DX12BufferAddr;
+					HRESULT Res = BufferInfo.GPUBuffer.Resource->Map(0, nullptr, (void**)&DX12BufferAddr);
+					check(SUCCEEDED(Res));
 
-							FElectraMediaDecoderOutputBufferPool_DX12::CopyWithPitchAdjust(DX12BufferAddr, ResourcePitch, BufferAddr, BufferPitch, DecodedPitch);
+					FElectraMediaDecoderOutputBufferPool_DX12::CopyWithPitchAdjust(DX12BufferAddr, ResourcePitch, BufferAddr, BufferPitch, DecodedHeight / BlockSizeY);
 
-							BufferResource->Unmap(0, nullptr);
+					BufferInfo.GPUBuffer.Resource->Unmap(0, nullptr);
 
-							BufferFence->Signal(BufferFenceValue);
-						}, TaskSync))
-					{
-						// note: we share one task sync for all buffers -> this will lead any external task to be executed after ALL tasks we trigger have been done (aka: after the "last" buffer is ready)
-						NewOutput->Buffer_TaskSync = TaskSync;
-					}
-					else
-					{
-						uint8* DX12BufferAddr;
-						HRESULT Res = BufferInfo.GPUBuffer.Resource->Map(0, nullptr, (void**)&DX12BufferAddr);
-						check(SUCCEEDED(Res));
-
-						FElectraMediaDecoderOutputBufferPool_DX12::CopyWithPitchAdjust(DX12BufferAddr, ResourcePitch, BufferAddr, BufferPitch, DecodedHeight / BlockSizeY);
-
-						BufferInfo.GPUBuffer.Resource->Unmap(0, nullptr);
-
-						BufferInfo.GPUBuffer.Fence->Signal(BufferInfo.GPUBuffer.FenceValue);
-					}
+					BufferInfo.GPUBuffer.Fence->Signal(BufferInfo.GPUBuffer.FenceValue);
 
 					BufferInfo.Pitch = ResourcePitch;
 				}
